LOD Paint warnings open without enabling LOD Paint
#fix check only LOD1 and above for determining whether or not a component conatains per-lod vertex colors
#misc merged functionality from Dev-Niagara:
	- Do not automatically apply vertex colors to all LODs anymore when switching out from per-lod painting state
	- Added text info box to inform user she's painting to all LOD levels
	- Added hotkey to cycle through LOD paint levels
	- Added button to propagate vertex paint to all LODs on mesh

diff --git a/Engine/Source/Editor/MeshPaint/Private/MeshPaintHelpers.cpp b/Engine/Source/Editor/MeshPaint/Private/MeshPaintHelpers.cpp
--- a/Engine/Source/Editor/MeshPaint/Private/MeshPaintHelpers.cpp
+++ b/Engine/Source/Editor/MeshPaint/Private/MeshPaintHelpers.cpp
@@ -1327,8 +1327,11 @@ bool MeshPaintHelpers::DoesMeshComponentContainPerLODColors(const UMeshComponent
 		if (SkeletalMesh)
 		{
 			const TArray<FSkeletalMeshLODInfo>& LODInfo = SkeletalMesh->GetLODInfoArray();
-			for ( const FSkeletalMeshLODInfo& Info : LODInfo )
+			// Only check LOD level 1 and above
+			const int32 NumLODs = SkeletalMesh->GetLODNum();
+			for (int32 LODIndex = 1; LODIndex < NumLODs; ++LODIndex)
 			{
+				const FSkeletalMeshLODInfo& Info = LODInfo[LODIndex];
 				if (Info.bHasPerLODVertexColors)
 				{
 					bPerLODColors = true;
diff --git a/Engine/Source/Editor/MeshPaintMode/Private/PaintModePainter.cpp b/Engine/Source/Editor/MeshPaintMode/Private/PaintModePainter.cpp
--- a/Engine/Source/Editor/MeshPaintMode/Private/PaintModePainter.cpp
+++ b/Engine/Source/Editor/MeshPaintMode/Private/PaintModePainter.cpp
@@ -432,6 +432,8 @@ void FPaintModePainter::PropagateVertexColorsToLODs()
 		}
 		else
 		{
+			// Reset the state flag as we'll be removing all per-lod colors 
+			bSelectionContainsPerLODColors = false;
 			//Remove painting on all lowers LODs before doing the propagation
 			for (UMeshComponent* SelectedComponent : PaintableComponents)
 			{
@@ -636,10 +638,10 @@ void FPaintModePainter::RegisterCommands(TSharedRef<FUICommandList> CommandList)
 			PaintSettings->VertexPaintSettings.EraseColor = Temp;
 		}
 	}));
-
+	
 	CommandList->MapAction(Commands.CycleToNextLOD, FExecuteAction::CreateRaw(this, &FPaintModePainter::CycleMeshLODs, 1));
 	CommandList->MapAction(Commands.CycleToPreviousLOD, FExecuteAction::CreateRaw(this, &FPaintModePainter::CycleMeshLODs, -1));
-
+	
 	/** Map commit texture painting to commiting all the outstanding paint changes */
 	auto CanCommitLambda = [this]() -> bool { return GetNumberOfPendingPaintChanges() > 0; };
 	CommandList->MapAction(Commands.CommitTexturePainting, FExecuteAction::CreateLambda([this]() { CommitAllPaintedTextures(); }), FCanExecuteAction::CreateLambda([this]() -> bool { return GetNumberOfPendingPaintChanges() > 0; }));
@@ -1115,7 +1117,7 @@ void FPaintModePainter::LODPaintStateChanged(const bool bLODPaintingEnabled)
 {
 	checkf(PaintSettings->PaintMode == EPaintMode::Vertices, TEXT("Can only change this state in vertex paint mode"));
 	bool AbortChange = false;
-	
+
 	// Set actual flag in the settings struct
 	PaintSettings->VertexPaintSettings.bPaintOnSpecificLOD = bLODPaintingEnabled;
 
@@ -1711,12 +1713,6 @@ void FPaintModePainter::FinishPaintingTexture()
 		PaintingTexture2D = nullptr;
 		TexturePaintingCurrentMeshComponent = nullptr;
 	}
-
-	if (bCachedForceLOD)
-	{
-		// Make assumption here that when we paint while having a forced lod the mesh will contain per-lod vertex colors after this
-		bSelectionContainsPerLODColors = true;
-	}
 }
 
 FPaintTexture2DData* FPaintModePainter::GetPaintTargetData(const UTexture2D* InTexture)
@@ -1999,16 +1995,16 @@ void FPaintModePainter::PropagateVertexColorsToAsset()
 				// Will not be guaranteed to match render data as user can paint to a specific LOD index
 				if (Component->LODData.IsValidIndex(LODIndex))
 				{
-					FStaticMeshComponentLODInfo& InstanceMeshLODInfo = Component->LODData[LODIndex];
-					if (InstanceMeshLODInfo.OverrideVertexColors)
-					{
-						Mesh->Modify();
-						// Try using the mapping generated when building the mesh.
-						if (MeshPaintHelpers::PropagateColorsToRawMesh(Mesh, LODIndex, InstanceMeshLODInfo))
-						{
-							SomePaintWasPropagated = true;
-						}
-					}
+				  FStaticMeshComponentLODInfo& InstanceMeshLODInfo = Component->LODData[LODIndex];
+				  if (InstanceMeshLODInfo.OverrideVertexColors)
+				  {
+					  Mesh->Modify();
+					  // Try using the mapping generated when building the mesh.
+					  if (MeshPaintHelpers::PropagateColorsToRawMesh(Mesh, LODIndex, InstanceMeshLODInfo))
+					  {
+						  SomePaintWasPropagated = true;
+					  }
+				  }
 				}
 			}
 
