Fix VS2019 compile error.

#rnx

diff --git a/Engine/Source/Editor/MeshPaintMode/Private/PaintModeSettingsCustomization.cpp b/Engine/Source/Editor/MeshPaintMode/Private/PaintModeSettingsCustomization.cpp
--- a/Engine/Source/Editor/MeshPaintMode/Private/PaintModeSettingsCustomization.cpp
+++ b/Engine/Source/Editor/MeshPaintMode/Private/PaintModeSettingsCustomization.cpp
@@ -371,7 +371,7 @@ void FVertexPaintSettingsCustomization::CustomizeChildren(TSharedRef<IPropertyHa
 						static const FText SkelMeshNotificationText = LOCTEXT("SkelMeshAssetPaintInfo", "Paint is propagated to Skeletal Mesh Asset(s)");
 						static const FText StaticMeshNotificationText = LOCTEXT("StaticMeshAssetPaintInfo", "Paint is applied to all LODs");
 
-						const bool bSkelMeshText = FPaintModePainter::Get()->GetSelectedComponents<USkeletalMeshComponent>().Num();
+						const bool bSkelMeshText = FPaintModePainter::Get()->GetSelectedComponents<USkeletalMeshComponent>().Num() > 0;
 						const bool bLODPaintText = !PaintSettings->bPaintOnSpecificLOD;
 						return FText::Format(FTextFormat::FromString(TEXT("{0}{1}{2}")), bSkelMeshText ? SkelMeshNotificationText : FText::GetEmpty(), bSkelMeshText && bLODPaintText ? FText::FromString(TEXT("\n")) : FText::GetEmpty(), bLODPaintText ? StaticMeshNotificationText : FText::GetEmpty());
 					})
