Only create vertex paint transaction if actually vertex painting

diff --git a/Engine/Source/Editor/MeshPaint/Private/IMeshPainter.cpp b/Engine/Source/Editor/MeshPaint/Private/IMeshPainter.cpp
--- a/Engine/Source/Editor/MeshPaint/Private/IMeshPainter.cpp
+++ b/Engine/Source/Editor/MeshPaint/Private/IMeshPainter.cpp
@@ -158,7 +158,10 @@ void IMeshPainter::FinishPainting()
 	if (bArePainting)
 	{
 		bArePainting = false;
-		EndTransaction();
+		if (PaintTransaction)
+		{
+			EndTransaction();
+		}
 	}
 
 	CurrentViewportInteractor = nullptr;
diff --git a/Engine/Source/Editor/MeshPaintMode/Private/PaintModePainter.cpp b/Engine/Source/Editor/MeshPaintMode/Private/PaintModePainter.cpp
--- a/Engine/Source/Editor/MeshPaintMode/Private/PaintModePainter.cpp
+++ b/Engine/Source/Editor/MeshPaintMode/Private/PaintModePainter.cpp
@@ -855,8 +855,12 @@ bool FPaintModePainter::PaintInternal(const FVector& InCameraOrigin, const FVect
 	if (HoveredComponents.Num() > 0)
 	{
 		if (bArePainting == false)
-		{
-			BeginTransaction(LOCTEXT("MeshPaintMode_VertexPaint_TransactionPaintStroke", "Vertex Paint"));
+		{ 
+			// Vertex painting is an ongoing transaction, while texture painting is handled separately later in a single transaction
+			if (PaintSettings->PaintMode == EPaintMode::Vertices)
+			{
+				BeginTransaction(LOCTEXT("MeshPaintMode_VertexPaint_TransactionPaintStroke", "Vertex Paint"));
+			}
 			bArePainting = true;
 			TimeSinceStartedPainting = 0.0f;
 		}
