Fix vertex color mode not showing how much memory instanced vertex colors are taking.

diff --git a/Engine/Source/Editor/MeshPaint/Private/MeshPaintHelpers.cpp b/Engine/Source/Editor/MeshPaint/Private/MeshPaintHelpers.cpp
--- a/Engine/Source/Editor/MeshPaint/Private/MeshPaintHelpers.cpp
+++ b/Engine/Source/Editor/MeshPaint/Private/MeshPaintHelpers.cpp
@@ -651,7 +651,7 @@ bool MeshPaintHelpers::RetrieveViewportPaintRays(const FSceneView* View, FViewpo
 	return false;
 }
 
-uint32 MeshPaintHelpers::GetVertexColorBufferSize(UStaticMeshComponent* MeshComponent, int32 LODIndex, bool bInstance)
+uint32 MeshPaintHelpers::GetVertexColorBufferSize(UMeshComponent* MeshComponent, int32 LODIndex, bool bInstance)
 {
 	checkf(MeshComponent != nullptr, TEXT("Invalid static mesh component ptr"));
 	
@@ -660,25 +660,50 @@ uint32 MeshPaintHelpers::GetVertexColorBufferSize(UStaticMeshComponent* MeshComp
 	// Retrieve component instance vertex color buffer size
 	if (bInstance)
 	{
-		if (MeshComponent->LODData.IsValidIndex(LODIndex))
+		if(UStaticMeshComponent* StaticMeshComponent = Cast<UStaticMeshComponent>(MeshComponent))
 		{
-			const FStaticMeshComponentLODInfo& InstanceMeshLODInfo = MeshComponent->LODData[LODIndex];
-			if (InstanceMeshLODInfo.OverrideVertexColors)
+			if (StaticMeshComponent->LODData.IsValidIndex(LODIndex))
 			{
-				SizeInBytes = InstanceMeshLODInfo.OverrideVertexColors->GetAllocatedSize();
+				const FStaticMeshComponentLODInfo& InstanceMeshLODInfo = StaticMeshComponent->LODData[LODIndex];
+				if (InstanceMeshLODInfo.OverrideVertexColors)
+				{
+					SizeInBytes = InstanceMeshLODInfo.OverrideVertexColors->GetAllocatedSize();
+				}
+			}
+		}
+		else if (USkinnedMeshComponent* SkinnedMeshComponent = Cast<USkinnedMeshComponent>(MeshComponent))
+		{
+			if (SkinnedMeshComponent->LODInfo.IsValidIndex(LODIndex))
+			{
+				const FSkelMeshComponentLODInfo& LODInfo = SkinnedMeshComponent->LODInfo[LODIndex];
+				if (LODInfo.OverrideVertexColors)
+				{
+					SizeInBytes = LODInfo.OverrideVertexColors->GetAllocatedSize();
+				}
 			}
 		}
 	}
 	// Retrieve static mesh asset vertex color buffer size
 	else
 	{
-		UStaticMesh* StaticMesh = MeshComponent->GetStaticMesh();
-		checkf(StaticMesh != nullptr, TEXT("Invalid static mesh ptr"));
-		if (StaticMesh->RenderData->LODResources.IsValidIndex(LODIndex))
+		if (UStaticMeshComponent* StaticMeshComponent = Cast<UStaticMeshComponent>(MeshComponent))
 		{
-			// count the base mesh color data
-			FStaticMeshLODResources& LODModel = StaticMesh->RenderData->LODResources[LODIndex];
-			SizeInBytes = LODModel.VertexBuffers.ColorVertexBuffer.GetAllocatedSize();
+			UStaticMesh* StaticMesh = StaticMeshComponent->GetStaticMesh();
+			checkf(StaticMesh != nullptr, TEXT("Invalid static mesh ptr"));
+			if (StaticMesh->RenderData->LODResources.IsValidIndex(LODIndex))
+			{
+				// count the base mesh color data
+				FStaticMeshLODResources& LODModel = StaticMesh->RenderData->LODResources[LODIndex];
+				SizeInBytes = LODModel.VertexBuffers.ColorVertexBuffer.GetAllocatedSize();
+			}
+		}
+		else if (USkinnedMeshComponent* SkinnedMeshComponent = Cast<USkinnedMeshComponent>(MeshComponent))
+		{
+			FSkeletalMeshRenderData* RenderData = SkinnedMeshComponent->GetSkeletalMeshRenderData();
+			if (RenderData && RenderData->LODRenderData.IsValidIndex(LODIndex))
+			{
+				SizeInBytes = RenderData->LODRenderData[LODIndex].StaticVertexBuffers.ColorVertexBuffer.GetAllocatedSize();
+			}
 		}
 	}
 	
diff --git a/Engine/Source/Editor/MeshPaint/Public/MeshPaintHelpers.h b/Engine/Source/Editor/MeshPaint/Public/MeshPaintHelpers.h
--- a/Engine/Source/Editor/MeshPaint/Public/MeshPaintHelpers.h
+++ b/Engine/Source/Editor/MeshPaint/Public/MeshPaintHelpers.h
@@ -53,7 +53,7 @@ public:
 	static bool PropagateColorsToRawMesh(UStaticMesh* StaticMesh, int32 LODIndex, FStaticMeshComponentLODInfo& ComponentLODInfo);	
 
 	/** Retrieves the Vertex Color buffer size for the given LOD level in the Static Mesh */
-	static uint32 GetVertexColorBufferSize(UStaticMeshComponent* MeshComponent, int32 LODIndex, bool bInstance);
+	static uint32 GetVertexColorBufferSize(UMeshComponent* MeshComponent, int32 LODIndex, bool bInstance);
 
 	/** Retrieves the vertex positions from the given LOD level in the Static Mesh */
 	static TArray<FVector> GetVerticesForLOD(const UStaticMesh* StaticMesh, int32 LODIndex);
diff --git a/Engine/Source/Editor/MeshPaintMode/Private/PaintModePainter.cpp b/Engine/Source/Editor/MeshPaintMode/Private/PaintModePainter.cpp
--- a/Engine/Source/Editor/MeshPaintMode/Private/PaintModePainter.cpp
+++ b/Engine/Source/Editor/MeshPaintMode/Private/PaintModePainter.cpp
@@ -387,6 +387,8 @@ void FPaintModePainter::PasteVertexColors()
 			}
 		}
 	}
+
+	UpdateCachedVertexDataSize();
 }
 
 void FPaintModePainter::FixVertexColors()
@@ -407,6 +409,8 @@ void FPaintModePainter::RemoveVertexColors()
 	{
 		MeshPaintHelpers::RemoveComponentInstanceVertexColors(Component);
 	}
+
+	UpdateCachedVertexDataSize();
 }
 
 void FPaintModePainter::PropagateVertexColorsToLODs()
@@ -496,6 +500,24 @@ void FPaintModePainter::CycleMeshLODs(int32 Direction)
 	}
 }
 
+void FPaintModePainter::UpdateCachedVertexDataSize()
+{
+	if (PaintSettings->PaintMode == EPaintMode::Vertices)
+	{
+		CachedVertexDataSize = 0;
+
+		const bool bInstance = true;
+		for (UMeshComponent* SelectedComponent : PaintableComponents)
+		{
+			int32 NumLODs = MeshPaintHelpers::GetNumberOfLODs(SelectedComponent);
+			for (int32 LODIndex = 0; LODIndex < NumLODs; ++LODIndex)
+			{
+				CachedVertexDataSize += MeshPaintHelpers::GetVertexColorBufferSize(SelectedComponent, LODIndex, bInstance);
+			}
+		}
+	}
+}
+
 bool FPaintModePainter::CanSaveMeshPackages() const
 {
 	// Check whether or not any of our selected mesh components contain mesh objects which require saving
@@ -783,6 +805,9 @@ void FPaintModePainter::FinishPainting()
 	{
 		PropagateVertexColorsToLODs();
 	}
+
+	UpdateCachedVertexDataSize();
+
 }
 
 bool FPaintModePainter::PaintInternal(const FVector& InCameraOrigin, const TArrayView<TPair<FVector, FVector>>& Rays, EMeshPaintAction PaintAction, float PaintStrength)
@@ -2168,7 +2193,9 @@ void FPaintModePainter::Tick(FEditorViewportClient* ViewportClient, float DeltaT
 		bRefreshCachedData = false;
 		CacheSelectionData();
 		CacheTexturePaintData();
-		
+
+		UpdateCachedVertexDataSize();
+
 		bDoRestoreRenTargets = true;
 	}
 
diff --git a/Engine/Source/Editor/MeshPaintMode/Private/PaintModePainter.h b/Engine/Source/Editor/MeshPaintMode/Private/PaintModePainter.h
--- a/Engine/Source/Editor/MeshPaintMode/Private/PaintModePainter.h
+++ b/Engine/Source/Editor/MeshPaintMode/Private/PaintModePainter.h
@@ -101,6 +101,10 @@ public:
 	virtual TSharedPtr<SWidget> GetWidget() override;
 	/** End IMeshPainter overrides */
 	
+	/**
+	 * Get the color buffer size for selected components being vertex painted
+	 */
+	int32 GetVertexPaintColorBufferSize() const { return CachedVertexDataSize; }
 protected:
 	/** Override from IMeshPainter used for applying actual painting */
 	virtual bool PaintInternal(const FVector& InCameraOrigin, const TArrayView<TPair<FVector, FVector>>& Rays, EMeshPaintAction PaintAction, float PaintStrength) override;
@@ -155,7 +159,8 @@ protected:
 	void RemoveVertexColors();
 	void PropagateVertexColorsToLODs();
 	void CycleMeshLODs(int32 Direction);
-	
+	void UpdateCachedVertexDataSize();
+
 	/** Checks whether or not the current selection contains components which reference the same (static/skeletal)-mesh */
 	bool ContainsDuplicateMeshes(TArray<UMeshComponent*>& Components) const;
 	
@@ -303,6 +308,7 @@ protected:
 	TArray<UMeshComponent*> PaintableComponents;
 	/** Contains copied vertex color data */
 	TArray<FPerComponentVertexColorData> CopiedColorsByComponent;
+	uint32 CachedVertexDataSize;
 	// End vertex paint state
 
 	/** UI command list object */
diff --git a/Engine/Source/Editor/MeshPaintMode/Private/PaintModeSettingsCustomization.cpp b/Engine/Source/Editor/MeshPaintMode/Private/PaintModeSettingsCustomization.cpp
--- a/Engine/Source/Editor/MeshPaintMode/Private/PaintModeSettingsCustomization.cpp
+++ b/Engine/Source/Editor/MeshPaintMode/Private/PaintModeSettingsCustomization.cpp
@@ -20,6 +20,8 @@
 #include "Widgets/Layout/SBox.h"
 #include "ScopedTransaction.h"
 
+#define LOCTEXT_NAMESPACE "PaintModePainter"
+
 TSharedRef<IPropertyTypeCustomization> FTexturePaintSettingsCustomization::MakeInstance()
 {
 	return MakeShareable(new FTexturePaintSettingsCustomization());
@@ -324,6 +326,14 @@ void FVertexPaintSettingsCustomization::CustomizeChildren(TSharedRef<IPropertyHa
 		.ValueContent()
 		[
 			SNew(SHorizontalBox)
+			+ SHorizontalBox::Slot()
+			.Padding(0.0f, 0.0f, 4.0f, 0.0f)
+			.AutoWidth()
+			[
+				SNew(SCheckBox)
+				.IsChecked_Lambda([=]() -> ECheckBoxState { return (PaintSettings->bPaintOnSpecificLOD ? ECheckBoxState::Checked : ECheckBoxState::Unchecked); })
+				.OnCheckStateChanged(FOnCheckStateChanged::CreateLambda([=](ECheckBoxState State) { FPaintModePainter::Get()->LODPaintStateChanged(State == ECheckBoxState::Checked); }))
+			]
 			+SHorizontalBox::Slot()
 			.Padding(0.0f, 0.0f, 4.0f, 0.0f)			
 			[
@@ -338,14 +348,35 @@ void FVertexPaintSettingsCustomization::CustomizeChildren(TSharedRef<IPropertyHa
 				.OnValueChanged(SNumericEntryBox<int32>::FOnValueChanged::CreateLambda([=](int32 Value) { PaintSettings->LODIndex = Value; }))
 				.OnValueCommitted(SNumericEntryBox<int32>::FOnValueCommitted::CreateLambda([=](int32 Value, ETextCommit::Type CommitType) { PaintSettings->LODIndex = Value; FPaintModePainter::Get()->PaintLODChanged(); }))
 			]
-			+SHorizontalBox::Slot()
-			.Padding(0.0f, 0.0f, 4.0f, 0.0f)
-			[
-				SNew(SCheckBox)
-				.IsChecked_Lambda([=]() -> ECheckBoxState { return (PaintSettings->bPaintOnSpecificLOD ? ECheckBoxState::Checked : ECheckBoxState::Unchecked); })
-				.OnCheckStateChanged(FOnCheckStateChanged::CreateLambda([=](ECheckBoxState State) { FPaintModePainter::Get()->LODPaintStateChanged(State == ECheckBoxState::Checked); }))
-			]
 		];
+
+		ChildBuilder.AddCustomRow(NSLOCTEXT("LODPainting", "LODPaintingLabel", "LOD Model Painting"))
+			.WholeRowContent()
+			[
+				SNew(SBorder)
+				.Visibility_Lambda([this]() -> EVisibility
+				{
+					return PaintSettings->bPaintOnSpecificLOD ? EVisibility::Collapsed : EVisibility::Visible;
+				})
+				.Padding(FMargin(4.0f))
+				.BorderImage(FEditorStyle::GetBrush("SettingsEditor.CheckoutWarningBorder"))
+				.BorderBackgroundColor(FColor(166, 137, 0))
+				[
+					SNew(STextBlock)
+					.AutoWrapText(true)
+					.Font(IDetailLayoutBuilder::GetDetailFont())
+
+					.Text_Lambda([this]() -> FText
+					{
+						static const FText SkelMeshNotificationText = LOCTEXT("SkelMeshAssetPaintInfo", "Paint is propagated to Skeletal Mesh Asset(s)");
+						static const FText StaticMeshNotificationText = LOCTEXT("StaticMeshAssetPaintInfo", "Paint is applied to all LODs");
+
+						const bool bSkelMeshText = FPaintModePainter::Get()->GetSelectedComponents<USkeletalMeshComponent>().Num();
+						const bool bLODPaintText = !PaintSettings->bPaintOnSpecificLOD;
+						return FText::Format(FTextFormat::FromString(TEXT("{0}{1}{2}")), bSkelMeshText ? SkelMeshNotificationText : FText::GetEmpty(), bSkelMeshText && bLODPaintText ? FText::FromString(TEXT("\n")) : FText::GetEmpty(), bLODPaintText ? StaticMeshNotificationText : FText::GetEmpty());
+					})
+				]
+			];
 }
 
 EVisibility FVertexPaintSettingsCustomization::ArePropertiesVisible(const int32 VisibleType) const
@@ -442,3 +473,5 @@ TSharedRef<SHorizontalBox> CreateColorChannelWidget(TSharedRef<IPropertyHandle>
 			ChannelProperty->CreatePropertyNameWidget()
 		];
 }
+
+#undef LOCTEXT_NAMESPACE
\ No newline at end of file
diff --git a/Engine/Source/Editor/MeshPaintMode/Private/SPaintModeWidget.cpp b/Engine/Source/Editor/MeshPaintMode/Private/SPaintModeWidget.cpp
--- a/Engine/Source/Editor/MeshPaintMode/Private/SPaintModeWidget.cpp
+++ b/Engine/Source/Editor/MeshPaintMode/Private/SPaintModeWidget.cpp
@@ -18,6 +18,7 @@
 
 #include "Modules/ModuleManager.h"
 #include "PaintModeCommands.h"
+#include "DetailLayoutBuilder.h"
 
 #define LOCTEXT_NAMESPACE "PaintModePainter"
 
@@ -62,13 +63,33 @@ void SPaintModeWidget::Construct(const FArguments& InArgs, FPaintModePainter* In
 			[
 				CreateTexturePaintWidget()->AsShared()
 			]
-
+			+ SVerticalBox::Slot()
+			.Padding(2.0f, 4.0f)
+			.HAlign(HAlign_Right)
+			.AutoHeight()
+			[
+				SNew(STextBlock)
+				.Font(IDetailLayoutBuilder::GetDetailFont())
+				.Visibility_Lambda([this]() -> EVisibility
+				{
+					return PaintModeSettings->PaintMode == EPaintMode::Vertices ? EVisibility::Visible : EVisibility::Collapsed;
+				})
+				.Text_Lambda([this]() -> FText
+				{
+					if (MeshPainter)
+					{
+						return FText::Format(FTextFormat::FromString(TEXT("Instance Color Size: {0} KB")), MeshPainter->GetVertexPaintColorBufferSize() / 1024.f);
+					}
+					return FText::GetEmpty();
+				})
+			]
 			/** DetailsView containing brush and paint settings */
 			+ SVerticalBox::Slot()
 			.AutoHeight()				
 			[
 				SettingsDetailsView->AsShared()
 			]
+
 		]
 	];
 }
@@ -103,9 +124,6 @@ TSharedPtr<SWidget> SPaintModeWidget::CreateVertexPaintWidget()
 	TSharedPtr<SHorizontalBox> VertexColorActionBox;
 	TSharedPtr<SHorizontalBox> InstanceColorActionBox;
 
-	static const FText SkelMeshNotificationText = LOCTEXT("SkelMeshAssetPaintInfo", "Paint is directly propagated to Skeletal Mesh Asset(s)");
-	static const FText StaticMeshNotificationText = LOCTEXT("StaticMeshAssetPaintInfo", "Paint is directly applied to all LODs");	
-		
 	SAssignNew(VertexColorWidget, SVerticalBox)
 	.Visibility(this, &SPaintModeWidget::IsVertexPaintModeVisible)
 	+ SVerticalBox::Slot()
@@ -116,58 +134,14 @@ TSharedPtr<SWidget> SPaintModeWidget::CreateVertexPaintWidget()
 		SAssignNew(VertexColorActionBox, SHorizontalBox)
 	]
 	
-	+SVerticalBox::Slot()
-	.AutoHeight()
-	.Padding(StandardPadding)	
-	.HAlign(HAlign_Center)
-	[
-		SAssignNew(InstanceColorActionBox, SHorizontalBox)
-	]
-
-	+SVerticalBox::Slot()
+	+ SVerticalBox::Slot()
 	.AutoHeight()
 	.Padding(StandardPadding)
-	.VAlign(VAlign_Center)
 	.HAlign(HAlign_Center)
 	[
-		SNew(SBorder)
-		.BorderImage(FEditorStyle::GetBrush("SettingsEditor.CheckoutWarningBorder"))
-		.BorderBackgroundColor(FColor(166,137,0))				
-		[
-			SNew(SHorizontalBox)
-			.Visibility_Lambda([this]() -> EVisibility 
-			{
-				bool bVisible = MeshPainter && MeshPainter->GetSelectedComponents<USkeletalMeshComponent>().Num();
-				bVisible |= ((PaintModeSettings->PaintMode == EPaintMode::Vertices) && !PaintModeSettings->VertexPaintSettings.bPaintOnSpecificLOD);
-				return bVisible ? EVisibility::Visible : EVisibility::Collapsed;
-			})
-		
-			+ SHorizontalBox::Slot()
-			.VAlign(VAlign_Center)
-			.AutoWidth()
-			.Padding(6.0f, 0.0f)
-			[
-				SNew(SImage)
-				.Image(FEditorStyle::GetBrush("ClassIcon.SkeletalMeshComponent"))
-			]
-		
-			+SHorizontalBox::Slot()
-			.VAlign(VAlign_Center)
-			.FillWidth(.8f)
-			.Padding(StandardPadding)
-			[
-				SNew(STextBlock)
-				.AutoWrapText(true)
-				.Text_Lambda([this]() -> FText
-				{
-					const bool bSkelMeshText = MeshPainter && MeshPainter->GetSelectedComponents<USkeletalMeshComponent>().Num();
-					const bool bLODPaintText = (PaintModeSettings->PaintMode == EPaintMode::Vertices) && !PaintModeSettings->VertexPaintSettings.bPaintOnSpecificLOD;
-					return FText::Format(FTextFormat::FromString(TEXT("{0}{1}{2}")), bSkelMeshText ? SkelMeshNotificationText : FText::GetEmpty(), bSkelMeshText && bLODPaintText ? FText::FromString(TEXT("\n")) : FText::GetEmpty(), bLODPaintText ? StaticMeshNotificationText : FText::GetEmpty());
-				})
-			]
-		]
+		SAssignNew(InstanceColorActionBox, SHorizontalBox)
 	];
-	
+
 	FToolBarBuilder ColorToolbarBuilder(MeshPainter->GetUICommandList(), FMultiBoxCustomization::None);
 	ColorToolbarBuilder.SetLabelVisibility(EVisibility::Collapsed);
 	ColorToolbarBuilder.AddToolBarButton(FPaintModeCommands::Get().Fill, NAME_None, FText::GetEmpty(), TAttribute<FText>(), FSlateIcon(FEditorStyle::GetStyleSetName(), "MeshPaint.Fill"));
