Mesh Painting: Various crashes (Cleanup AddReferencedObjects code)

#rnx

diff --git a/Engine/Plugins/2D/Paper2D/Source/Paper2DEditor/Private/MeshPainting/MeshPaintSpriteAdapter.h b/Engine/Plugins/2D/Paper2D/Source/Paper2DEditor/Private/MeshPainting/MeshPaintSpriteAdapter.h
--- a/Engine/Plugins/2D/Paper2D/Source/Paper2DEditor/Private/MeshPainting/MeshPaintSpriteAdapter.h
+++ b/Engine/Plugins/2D/Paper2D/Source/Paper2DEditor/Private/MeshPainting/MeshPaintSpriteAdapter.h
@@ -20,6 +20,8 @@ class FMeshPaintSpriteAdapter : public IMeshPaintGeometryAdapter
 {
 public:
 	static void InitializeAdapterGlobals() {}
+	static void AddReferencedObjectsGlobals(FReferenceCollector& Collector) {}
+	static void CleanupGlobals() {}
 
 	virtual bool Construct(UMeshComponent* InComponent, int32 InMeshLODIndex) override;
 	virtual bool Initialize() override;
@@ -67,4 +69,6 @@ class FMeshPaintSpriteAdapterFactory : public IMeshPaintGeometryAdapterFactory
 public:
 	virtual TSharedPtr<IMeshPaintGeometryAdapter> Construct(UMeshComponent* InComponent, int32 InMeshLODIndex) const override;
 	virtual void InitializeAdapterGlobals() override { FMeshPaintSpriteAdapter::InitializeAdapterGlobals(); }
+	virtual void AddReferencedObjectsGlobals(FReferenceCollector& Collector) override { FMeshPaintSpriteAdapter::AddReferencedObjectsGlobals(Collector); }
+	virtual void CleanupGlobals() override { FMeshPaintSpriteAdapter::CleanupGlobals(); }
 };
diff --git a/Engine/Source/Editor/ClothPainter/Private/ClothMeshAdapter.h b/Engine/Source/Editor/ClothPainter/Private/ClothMeshAdapter.h
--- a/Engine/Source/Editor/ClothPainter/Private/ClothMeshAdapter.h
+++ b/Engine/Source/Editor/ClothPainter/Private/ClothMeshAdapter.h
@@ -37,6 +37,8 @@ protected:
 	};
 public:
 	static void InitializeAdapterGlobals() {}
+	static void AddReferencedObjectsGlobals(FReferenceCollector& Collector) {}
+	static void CleanupGlobals() {}
 
 	virtual bool Construct(UMeshComponent* InComponent, int32 InPaintingMeshLODIndex) override;
 	virtual bool Initialize() override;
@@ -123,4 +125,6 @@ class FClothMeshPaintAdapterFactory : public IMeshPaintGeometryAdapterFactory
 public:
 	virtual TSharedPtr<IMeshPaintGeometryAdapter> Construct(UMeshComponent* InComponent, int32 InPaintingMeshLODIndex) const override;
 	virtual void InitializeAdapterGlobals() override { FClothMeshPaintAdapter::InitializeAdapterGlobals(); }
+	virtual void AddReferencedObjectsGlobals(FReferenceCollector& Collector) override { FClothMeshPaintAdapter::AddReferencedObjectsGlobals(Collector); }
+	virtual void CleanupGlobals() override { FClothMeshPaintAdapter::CleanupGlobals(); }
 };
diff --git a/Engine/Source/Editor/MeshPaint/Private/MeshPaintAdapterFactory.cpp b/Engine/Source/Editor/MeshPaint/Private/MeshPaintAdapterFactory.cpp
--- a/Engine/Source/Editor/MeshPaint/Private/MeshPaintAdapterFactory.cpp
+++ b/Engine/Source/Editor/MeshPaint/Private/MeshPaintAdapterFactory.cpp
@@ -29,3 +29,19 @@ void FMeshPaintAdapterFactory::InitializeAdapterGlobals()
 		Factory->InitializeAdapterGlobals();
 	}
 }
+
+void FMeshPaintAdapterFactory::AddReferencedObjectsGlobals(FReferenceCollector& Collector)
+{
+	for (const auto& Factory : FactoryList)
+	{
+		Factory->AddReferencedObjectsGlobals(Collector);
+	}
+}
+
+void FMeshPaintAdapterFactory::CleanupGlobals()
+{
+	for (const auto& Factory : FactoryList)
+	{
+		Factory->CleanupGlobals();
+	}
+}
diff --git a/Engine/Source/Editor/MeshPaint/Private/MeshPaintSkeletalMeshAdapter.cpp b/Engine/Source/Editor/MeshPaint/Private/MeshPaintSkeletalMeshAdapter.cpp
--- a/Engine/Source/Editor/MeshPaint/Private/MeshPaintSkeletalMeshAdapter.cpp
+++ b/Engine/Source/Editor/MeshPaint/Private/MeshPaintSkeletalMeshAdapter.cpp
@@ -109,6 +109,19 @@ void FMeshPaintGeometryAdapterForSkeletalMeshes::InitializeAdapterGlobals()
 	}
 }
 
+void FMeshPaintGeometryAdapterForSkeletalMeshes::CleanupGlobals()
+{
+	for (TPair<USkeletalMesh*, FSkeletalMeshReferencers>& Pair : MeshToComponentMap)
+	{
+		if (Pair.Key && Pair.Value.RestoreBodySetup)
+		{
+			Pair.Key->BodySetup = Pair.Value.RestoreBodySetup;
+		}
+	}
+
+	MeshToComponentMap.Empty();
+}
+
 void FMeshPaintGeometryAdapterForSkeletalMeshes::OnAdded()
 {
 	checkf(SkeletalMeshComponent, TEXT("Invalid SkeletalMesh Component"));
@@ -171,11 +184,9 @@ void FMeshPaintGeometryAdapterForSkeletalMeshes::OnAdded()
 
 void FMeshPaintGeometryAdapterForSkeletalMeshes::OnRemoved()
 {
-	checkf(SkeletalMeshComponent, TEXT("Invalid SkeletalMesh Component"));
-	
 	// If the referenced skeletal mesh has been destroyed (and nulled by GC), don't try to do anything more.
 	// It should be in the process of removing all global geometry adapters if it gets here in this situation.
-	if (!ReferencedSkeletalMesh)
+	if (!ReferencedSkeletalMesh || !SkeletalMeshComponent)
 	{
 		return;
 	}
@@ -286,25 +297,23 @@ void FMeshPaintGeometryAdapterForSkeletalMeshes::ApplyOrRemoveTextureOverride(UT
 	DefaultApplyOrRemoveTextureOverride(SkeletalMeshComponent, SourceTexture, OverrideTexture);
 }
 
-void FMeshPaintGeometryAdapterForSkeletalMeshes::AddReferencedObjects(FReferenceCollector& Collector)
+void FMeshPaintGeometryAdapterForSkeletalMeshes::AddReferencedObjectsGlobals(FReferenceCollector& Collector)
 {
-	if (!ReferencedSkeletalMesh)
-	{
-		return;
-	}
-
-	FSkeletalMeshReferencers* SkeletalMeshReferencers = MeshToComponentMap.Find(ReferencedSkeletalMesh);
-	checkf(SkeletalMeshReferencers, TEXT("No references found for Skeletal Mesh"));
-	if (SkeletalMeshReferencers->RestoreBodySetup != nullptr)
+	for (TPair<USkeletalMesh*, FSkeletalMeshReferencers>& Pair : MeshToComponentMap)
 	{
-		Collector.AddReferencedObject(SkeletalMeshReferencers->RestoreBodySetup);
+		Collector.AddReferencedObject(Pair.Key);
+		Collector.AddReferencedObject(Pair.Value.RestoreBodySetup);
+		for (FSkeletalMeshReferencers::FReferencersInfo& ReferencerInfo : Pair.Value.Referencers)
+		{
+			Collector.AddReferencedObject(ReferencerInfo.SkeletalMeshComponent);
+		}
 	}
-	
+}
 
-	for (auto& Info : SkeletalMeshReferencers->Referencers)
-	{
-		Collector.AddReferencedObject(Info.SkeletalMeshComponent);
-	}
+void FMeshPaintGeometryAdapterForSkeletalMeshes::AddReferencedObjects(FReferenceCollector& Collector)
+{
+	Collector.AddReferencedObject(ReferencedSkeletalMesh);
+	Collector.AddReferencedObject(SkeletalMeshComponent);
 }
 
 void FMeshPaintGeometryAdapterForSkeletalMeshes::GetTextureCoordinate(int32 VertexIndex, int32 ChannelIndex, FVector2D& OutTextureCoordinate) const
diff --git a/Engine/Source/Editor/MeshPaint/Private/MeshPaintStaticMeshAdapter.cpp b/Engine/Source/Editor/MeshPaint/Private/MeshPaintStaticMeshAdapter.cpp
--- a/Engine/Source/Editor/MeshPaint/Private/MeshPaintStaticMeshAdapter.cpp
+++ b/Engine/Source/Editor/MeshPaint/Private/MeshPaintStaticMeshAdapter.cpp
@@ -143,6 +143,19 @@ void FMeshPaintGeometryAdapterForStaticMeshes::InitializeAdapterGlobals()
 	}
 }
 
+void FMeshPaintGeometryAdapterForStaticMeshes::CleanupGlobals()
+{
+	for (TPair<UStaticMesh*, FStaticMeshReferencers>& Pair : MeshToComponentMap)
+	{
+		if (Pair.Key && Pair.Value.RestoreBodySetup)
+		{
+			Pair.Key->BodySetup = Pair.Value.RestoreBodySetup;
+		}
+	}
+
+	MeshToComponentMap.Empty();
+}
+
 void FMeshPaintGeometryAdapterForStaticMeshes::OnAdded()
 {
 	check(StaticMeshComponent);
@@ -205,11 +218,9 @@ void FMeshPaintGeometryAdapterForStaticMeshes::OnAdded()
 
 void FMeshPaintGeometryAdapterForStaticMeshes::OnRemoved()
 {
-	check(StaticMeshComponent);
-	
 	// If the referenced static mesh has been destroyed (and nulled by GC), don't try to do anything more.
 	// It should be in the process of removing all global geometry adapters if it gets here in this situation.
-	if (!ReferencedStaticMesh)
+	if (!ReferencedStaticMesh || !StaticMeshComponent)
 	{
 		return;
 	}
@@ -233,13 +244,6 @@ void FMeshPaintGeometryAdapterForStaticMeshes::OnRemoved()
 			StaticMeshComponent->RecreatePhysicsState();
 
 			StaticMeshReferencers->Referencers.RemoveAtSwap(Index);
-
-			// If the last reference was removed, restore the body setup for the static mesh
-			if (StaticMeshReferencers->Referencers.Num() == 0)
-			{
-				ReferencedStaticMesh->BodySetup = StaticMeshReferencers->RestoreBodySetup;
-				verify(MeshToComponentMap.Remove(ReferencedStaticMesh) == 1);
-			}
 		}
 		else
 		{
@@ -249,6 +253,13 @@ void FMeshPaintGeometryAdapterForStaticMeshes::OnRemoved()
 				return Info.StaticMeshComponent == nullptr;
 			});
 		}
+
+		// If the last reference was removed, restore the body setup for the static mesh
+		if (StaticMeshReferencers->Referencers.Num() == 0)
+		{
+			ReferencedStaticMesh->BodySetup = StaticMeshReferencers->RestoreBodySetup;
+			verify(MeshToComponentMap.Remove(ReferencedStaticMesh) == 1);
+		}
 	}
 }
 
@@ -268,21 +279,23 @@ void FMeshPaintGeometryAdapterForStaticMeshes::ApplyOrRemoveTextureOverride(UTex
 	DefaultApplyOrRemoveTextureOverride(StaticMeshComponent, SourceTexture, OverrideTexture);
 }
 
-void FMeshPaintGeometryAdapterForStaticMeshes::AddReferencedObjects(FReferenceCollector& Collector)
+void FMeshPaintGeometryAdapterForStaticMeshes::AddReferencedObjectsGlobals(FReferenceCollector& Collector)
 {
-	if (!ReferencedStaticMesh)
+	for (TPair<UStaticMesh*, FStaticMeshReferencers>& Pair : MeshToComponentMap)
 	{
-		return;
+		Collector.AddReferencedObject(Pair.Key);
+		Collector.AddReferencedObject(Pair.Value.RestoreBodySetup);
+		for (FStaticMeshReferencers::FReferencersInfo& ReferencerInfo : Pair.Value.Referencers)
+		{
+			Collector.AddReferencedObject(ReferencerInfo.StaticMeshComponent);
+		}
 	}
+}
 
-	FStaticMeshReferencers* StaticMeshReferencers = MeshToComponentMap.Find(ReferencedStaticMesh);
-	check(StaticMeshReferencers);
-	Collector.AddReferencedObject(StaticMeshReferencers->RestoreBodySetup);
-
-	for (auto& Info : StaticMeshReferencers->Referencers)
-	{
-		Collector.AddReferencedObject(Info.StaticMeshComponent);
-	}
+void FMeshPaintGeometryAdapterForStaticMeshes::AddReferencedObjects(FReferenceCollector& Collector)
+{
+	Collector.AddReferencedObject(ReferencedStaticMesh);
+	Collector.AddReferencedObject(StaticMeshComponent);
 }
 
 void FMeshPaintGeometryAdapterForStaticMeshes::GetVertexColor(int32 VertexIndex, FColor& OutColor, bool bInstance /*= true*/) const
diff --git a/Engine/Source/Editor/MeshPaint/Public/IMeshPaintGeometryAdapterFactory.h b/Engine/Source/Editor/MeshPaint/Public/IMeshPaintGeometryAdapterFactory.h
--- a/Engine/Source/Editor/MeshPaint/Public/IMeshPaintGeometryAdapterFactory.h
+++ b/Engine/Source/Editor/MeshPaint/Public/IMeshPaintGeometryAdapterFactory.h
@@ -4,6 +4,7 @@
 
 class IMeshPaintGeometryAdapter;
 class UMeshComponent;
+class FReferenceCollector;
 
 /**
  * Factory for IMeshPaintGeometryAdapter
@@ -13,5 +14,7 @@ class IMeshPaintGeometryAdapterFactory
 public:
 	virtual TSharedPtr<IMeshPaintGeometryAdapter> Construct(UMeshComponent* InComponent, int32 InMeshLODIndex) const = 0;
 	virtual void InitializeAdapterGlobals() = 0;
+	virtual void AddReferencedObjectsGlobals(FReferenceCollector& Collector) = 0;
+	virtual void CleanupGlobals() = 0;
 	virtual ~IMeshPaintGeometryAdapterFactory() {}
 };
\ No newline at end of file
diff --git a/Engine/Source/Editor/MeshPaint/Public/MeshPaintAdapterFactory.h b/Engine/Source/Editor/MeshPaint/Public/MeshPaintAdapterFactory.h
--- a/Engine/Source/Editor/MeshPaint/Public/MeshPaintAdapterFactory.h
+++ b/Engine/Source/Editor/MeshPaint/Public/MeshPaintAdapterFactory.h
@@ -7,6 +7,7 @@
 
 class UMeshComponent;
 class IMeshPaintGeometryAdapterFactory;
+class FReferenceCollector;
 
 class MESHPAINT_API FMeshPaintAdapterFactory
 {
@@ -16,4 +17,6 @@ public:
 public:
 	static TSharedPtr<class IMeshPaintGeometryAdapter> CreateAdapterForMesh(UMeshComponent* InComponent, int32 InPaintingMeshLODIndex);
 	static void InitializeAdapterGlobals();
+	static void AddReferencedObjectsGlobals(FReferenceCollector& Collector);
+	static void CleanupGlobals();
 };
diff --git a/Engine/Source/Editor/MeshPaint/Public/MeshPaintSkeletalMeshAdapter.h b/Engine/Source/Editor/MeshPaint/Public/MeshPaintSkeletalMeshAdapter.h
--- a/Engine/Source/Editor/MeshPaint/Public/MeshPaintSkeletalMeshAdapter.h
+++ b/Engine/Source/Editor/MeshPaint/Public/MeshPaintSkeletalMeshAdapter.h
@@ -24,6 +24,8 @@ class FMeshPaintGeometryAdapterForSkeletalMeshes : public FBaseMeshPaintGeometry
 {
 public:
 	static void InitializeAdapterGlobals();
+	static void AddReferencedObjectsGlobals(FReferenceCollector& Collector);
+	static void CleanupGlobals();
 
 	/** Start IMeshPaintGeometryAdapter Overrides */
 	virtual bool Construct(UMeshComponent* InComponent, int32 InMeshLODIndex) override;
@@ -106,4 +108,6 @@ class FMeshPaintGeometryAdapterForSkeletalMeshesFactory : public IMeshPaintGeome
 public:
 	virtual TSharedPtr<IMeshPaintGeometryAdapter> Construct(UMeshComponent* InComponent, int32 InMeshLODIndex) const override;
 	virtual void InitializeAdapterGlobals() override { FMeshPaintGeometryAdapterForSkeletalMeshes::InitializeAdapterGlobals(); }
+	virtual void AddReferencedObjectsGlobals(FReferenceCollector& Collector) override { FMeshPaintGeometryAdapterForSkeletalMeshes::AddReferencedObjectsGlobals(Collector); }
+	virtual void CleanupGlobals() override { FMeshPaintGeometryAdapterForSkeletalMeshes::CleanupGlobals(); }
 };
diff --git a/Engine/Source/Editor/MeshPaint/Public/MeshPaintStaticMeshAdapter.h b/Engine/Source/Editor/MeshPaint/Public/MeshPaintStaticMeshAdapter.h
--- a/Engine/Source/Editor/MeshPaint/Public/MeshPaintStaticMeshAdapter.h
+++ b/Engine/Source/Editor/MeshPaint/Public/MeshPaintStaticMeshAdapter.h
@@ -20,6 +20,8 @@ class FMeshPaintGeometryAdapterForStaticMeshes : public FBaseMeshPaintGeometryAd
 {
 public:
 	static void InitializeAdapterGlobals();
+	static void AddReferencedObjectsGlobals(FReferenceCollector& Collector);
+	static void CleanupGlobals();
 
 	/** Start IMeshPaintGeometryAdapter Overrides */
 	virtual bool Construct(UMeshComponent* InComponent, int32 InMeshLODIndex) override;
@@ -95,4 +97,6 @@ class FMeshPaintGeometryAdapterForStaticMeshesFactory : public IMeshPaintGeometr
 public:
 	virtual TSharedPtr<IMeshPaintGeometryAdapter> Construct(class UMeshComponent* InComponent, int32 InMeshLODIndex) const override;
 	virtual void InitializeAdapterGlobals() override { FMeshPaintGeometryAdapterForStaticMeshes::InitializeAdapterGlobals(); }
+	virtual void AddReferencedObjectsGlobals(FReferenceCollector& Collector) override { FMeshPaintGeometryAdapterForStaticMeshes::AddReferencedObjectsGlobals(Collector); }
+	virtual void CleanupGlobals() override { FMeshPaintGeometryAdapterForStaticMeshes::CleanupGlobals(); }
 };
diff --git a/Engine/Source/Editor/MeshPaintMode/Private/PaintModePainter.cpp b/Engine/Source/Editor/MeshPaintMode/Private/PaintModePainter.cpp
--- a/Engine/Source/Editor/MeshPaintMode/Private/PaintModePainter.cpp
+++ b/Engine/Source/Editor/MeshPaintMode/Private/PaintModePainter.cpp
@@ -57,7 +57,7 @@ FPaintModePainter::FPaintModePainter()
 FPaintModePainter::~FPaintModePainter()
 {
 	FCoreUObjectDelegates::OnObjectPropertyChanged.RemoveAll(this);
-	ComponentToAdapterMap.Empty();
+	Cleanup();
 	ComponentToTexturePaintSettingsMap.Empty();
 }
 
@@ -480,9 +480,12 @@ void FPaintModePainter::PropagateVertexColorsToLODs()
 
 	for (UMeshComponent* SelectedComponent : PaintableComponents)
 	{
-		TSharedPtr<IMeshPaintGeometryAdapter> MeshAdapter = ComponentToAdapterMap.FindChecked(SelectedComponent);
-		MeshPaintHelpers::ApplyVertexColorsToAllLODs(*MeshAdapter, SelectedComponent);
-		FComponentReregisterContext ReregisterContext(SelectedComponent);
+		if (SelectedComponent)
+		{
+			TSharedPtr<IMeshPaintGeometryAdapter> MeshAdapter = ComponentToAdapterMap.FindChecked(SelectedComponent);
+			MeshPaintHelpers::ApplyVertexColorsToAllLODs(*MeshAdapter, SelectedComponent);
+			FComponentReregisterContext ReregisterContext(SelectedComponent);
+		}
 	}
 
 	Refresh();
@@ -509,10 +512,13 @@ void FPaintModePainter::UpdateCachedVertexDataSize()
 		const bool bInstance = true;
 		for (UMeshComponent* SelectedComponent : PaintableComponents)
 		{
-			int32 NumLODs = MeshPaintHelpers::GetNumberOfLODs(SelectedComponent);
-			for (int32 LODIndex = 0; LODIndex < NumLODs; ++LODIndex)
+			if (SelectedComponent)
 			{
-				CachedVertexDataSize += MeshPaintHelpers::GetVertexColorBufferSize(SelectedComponent, LODIndex, bInstance);
+				int32 NumLODs = MeshPaintHelpers::GetNumberOfLODs(SelectedComponent);
+				for (int32 LODIndex = 0; LODIndex < NumLODs; ++LODIndex)
+				{
+					CachedVertexDataSize += MeshPaintHelpers::GetVertexColorBufferSize(SelectedComponent, LODIndex, bInstance);
+				}
 			}
 		}
 	}
@@ -694,17 +700,20 @@ const FHitResult FPaintModePainter::GetHitResult(const FVector& Origin, const FV
 
 		for (UMeshComponent* MeshComponent : PaintableComponents)
 		{
-			TSharedPtr<IMeshPaintGeometryAdapter> MeshAdapter = ComponentToAdapterMap.FindChecked(MeshComponent);
+			if (MeshComponent)
+			{
+				TSharedPtr<IMeshPaintGeometryAdapter> MeshAdapter = ComponentToAdapterMap.FindChecked(MeshComponent);
 
-			// Ray trace
-			FHitResult TraceHitResult(1.0f);
+				// Ray trace
+				FHitResult TraceHitResult(1.0f);
 
-			if (MeshAdapter->LineTraceComponent(TraceHitResult, TraceStart, TraceEnd, FCollisionQueryParams(SCENE_QUERY_STAT(Paint), true)))
-			{
-				// Find the closest impact
-				if ((BestTraceResult.GetComponent() == nullptr) || (TraceHitResult.Time < BestTraceResult.Time))
+				if (MeshAdapter->LineTraceComponent(TraceHitResult, TraceStart, TraceEnd, FCollisionQueryParams(SCENE_QUERY_STAT(Paint), true)))
 				{
-					BestTraceResult = TraceHitResult;
+					// Find the closest impact
+					if ((BestTraceResult.GetComponent() == nullptr) || (TraceHitResult.Time < BestTraceResult.Time))
+					{
+						BestTraceResult = TraceHitResult;
+					}
 				}
 			}
 		}
@@ -782,12 +791,15 @@ void FPaintModePainter::AddReferencedObjects(FReferenceCollector& Collector)
 	Collector.AddReferencedObject(BrushRenderTargetTexture);
 	Collector.AddReferencedObject(BrushMaskRenderTargetTexture);
 	Collector.AddReferencedObject(SeamMaskRenderTargetTexture);
+	Collector.AddReferencedObjects(PaintableComponents);
 	for (TMap< UTexture2D*, FPaintTexture2DData >::TIterator It(PaintTargetData); It; ++It)
 	{
 		Collector.AddReferencedObject(It.Key());
 		It.Value().AddReferencedObjects(Collector);
 	}
 
+	FMeshPaintAdapterFactory::AddReferencedObjectsGlobals(Collector);
+
 	for (TMap< UMeshComponent*, TSharedPtr<IMeshPaintGeometryAdapter>>::TIterator It(ComponentToAdapterMap); It; ++It)
 	{
 		Collector.AddReferencedObject(It.Key());
@@ -842,6 +854,11 @@ bool FPaintModePainter::PaintInternal(const FVector& InCameraOrigin, const TArra
 
 		for (UMeshComponent* MeshComponent : PaintableComponents)
 		{
+			if (!MeshComponent)
+			{
+				continue;
+			}
+
 			TSharedPtr<IMeshPaintGeometryAdapter>* MeshAdapterPtr = ComponentToAdapterMap.Find(MeshComponent);
 			if (!MeshAdapterPtr)
 			{
@@ -881,13 +898,16 @@ bool FPaintModePainter::PaintInternal(const FVector& InCameraOrigin, const TArra
 				// Vertex paint mode, so we want all valid components overlapping the brush hit location
 				for (auto TestComponent : PaintableComponents)
 				{
-					const FBox ComponentBounds = TestComponent->Bounds.GetBox();
-
-					if (ComponentToAdapterMap.Contains(TestComponent) && ComponentBounds.Intersect(BrushBounds))
+					if (TestComponent)
 					{
-						// OK, this mesh potentially overlaps the brush!
-						HoveredComponents.FindOrAdd(TestComponent).Add(i);
-						bUsed = true;
+						const FBox ComponentBounds = TestComponent->Bounds.GetBox();
+
+						if (ComponentToAdapterMap.Contains(TestComponent) && ComponentBounds.Intersect(BrushBounds))
+						{
+							// OK, this mesh potentially overlaps the brush!
+							HoveredComponents.FindOrAdd(TestComponent).Add(i);
+							bUsed = true;
+						}
 					}
 				}
 			}
@@ -1135,11 +1155,7 @@ void FPaintModePainter::Reset()
 	TexturePaintTargetList.Empty();
 	
 	// Cleanup all cached 
-	for (auto MeshAdapterPair : ComponentToAdapterMap)
-	{
-		MeshAdapterPair.Value->OnRemoved();
-	}
-	ComponentToAdapterMap.Empty();
+	Cleanup();
 }
 
 TSharedPtr<IMeshPaintGeometryAdapter> FPaintModePainter::GetMeshAdapterForComponent(const UMeshComponent* Component)
@@ -1227,7 +1243,10 @@ void FPaintModePainter::LODPaintStateChanged(const bool bLODPaintingEnabled)
 	//Make sure all static mesh render is dirty since we change the force LOD
 	for (UMeshComponent* SelectedComponent : PaintableComponents)
 	{
-		ComponentReregisterContext.Reset(new FComponentReregisterContext(SelectedComponent));
+		if (SelectedComponent)
+		{
+			ComponentReregisterContext.Reset(new FComponentReregisterContext(SelectedComponent));
+		}
 	}
 
 	Refresh();
@@ -1245,7 +1264,10 @@ void FPaintModePainter::PaintLODChanged()
 		//Make sure all static mesh render is dirty since we change the force LOD
 		for (UMeshComponent* SelectedComponent : PaintableComponents)
 		{
-			ComponentReregisterContext.Reset(new FComponentReregisterContext(SelectedComponent));
+			if (SelectedComponent)
+			{
+				ComponentReregisterContext.Reset(new FComponentReregisterContext(SelectedComponent));
+			}
 		}
 
 		Refresh();
@@ -1254,7 +1276,7 @@ void FPaintModePainter::PaintLODChanged()
 
 int32 FPaintModePainter::GetMaxUVIndexToPaint() const
 {
-	if (PaintableComponents.Num() == 1)
+	if (PaintableComponents.Num() == 1 && PaintableComponents[0])
 	{
 		return MeshPaintHelpers::GetNumberOfUVs(PaintableComponents[0], CachedLODIndex) - 1;
 	}
@@ -2000,7 +2022,10 @@ void FPaintModePainter::ApplyForcedLODIndex(int32 ForcedLODIndex)
 {	
 	for (UMeshComponent* SelectedComponent : PaintableComponents)
 	{
-		MeshPaintHelpers::ForceRenderMeshLOD(SelectedComponent, ForcedLODIndex);
+		if (SelectedComponent)
+		{
+			MeshPaintHelpers::ForceRenderMeshLOD(SelectedComponent, ForcedLODIndex);
+		}
 	}
 }
 
@@ -2184,13 +2209,19 @@ void FPaintModePainter::Refresh()
 {
 	// Ensure that we call OnRemoved while adapter/components are still valid
 	PaintableComponents.Empty();
+	Cleanup();
+
+	bRefreshCachedData = true;
+}
+
+void FPaintModePainter::Cleanup()
+{
 	for (auto MeshAdapterPair : ComponentToAdapterMap)
 	{
 		MeshAdapterPair.Value->OnRemoved();
 	}
 	ComponentToAdapterMap.Empty();
-
-	bRefreshCachedData = true;
+	FMeshPaintAdapterFactory::CleanupGlobals();
 }
 
 void FPaintModePainter::Tick(FEditorViewportClient* ViewportClient, float DeltaTime)
@@ -2209,7 +2240,7 @@ void FPaintModePainter::Tick(FEditorViewportClient* ViewportClient, float DeltaT
 	}
 
 	// Will set the texture override up for the selected texture, important for the drop down combo-list and selecting between material instances.
-	if (PaintSettings->PaintMode == EPaintMode::Textures && PaintableComponents.Num() == 1 && PaintSettings->TexturePaintSettings.PaintTexture)
+	if (PaintSettings->PaintMode == EPaintMode::Textures && PaintableComponents.Num() == 1 && PaintableComponents[0] && PaintSettings->TexturePaintSettings.PaintTexture)
 	{
 		for (UMeshComponent* MeshComponent : PaintableComponents)
 		{
@@ -2309,7 +2340,7 @@ void FPaintModePainter::CacheTexturePaintData()
 	TArray<UMeshComponent*> SelectedMeshComponents = GetSelectedComponents<UMeshComponent>();
 
 	PaintableTextures.Empty();
-	if (PaintableComponents.Num() == 1)
+	if (PaintableComponents.Num() == 1 && PaintableComponents[0])
 	{
 		const UMeshComponent* Component = PaintableComponents[0];
 		TSharedPtr<IMeshPaintGeometryAdapter> Adapter = ComponentToAdapterMap.FindChecked(Component);
diff --git a/Engine/Source/Editor/MeshPaintMode/Private/PaintModePainter.h b/Engine/Source/Editor/MeshPaintMode/Private/PaintModePainter.h
--- a/Engine/Source/Editor/MeshPaintMode/Private/PaintModePainter.h
+++ b/Engine/Source/Editor/MeshPaintMode/Private/PaintModePainter.h
@@ -249,6 +249,7 @@ private:
 	void ImportVertexColors();
 	void SavePaintedAssets();
 	void SaveModifiedTextures();
+	void Cleanup();
 protected:	
 	/** Texture paint state */
 	/** Textures eligible for painting retrieved from the current selection */
