Fixes Crash when creating new level in Paint Mode

diff --git a/Engine/Source/Editor/MeshPaintMode/Private/PaintModePainter.cpp b/Engine/Source/Editor/MeshPaintMode/Private/PaintModePainter.cpp
--- a/Engine/Source/Editor/MeshPaintMode/Private/PaintModePainter.cpp
+++ b/Engine/Source/Editor/MeshPaintMode/Private/PaintModePainter.cpp
@@ -1154,6 +1154,8 @@ void FPaintModePainter::Reset()
 	// Remove any existing texture targets
 	TexturePaintTargetList.Empty();
 	
+	PaintableComponents.Empty();
+
 	// Cleanup all cached 
 	Cleanup();
 }
