Fix during Mesh Painting in Mobile Preview

diff --git a/Engine/Source/Editor/MeshPaintMode/Private/MeshPaintEdMode.cpp b/Engine/Source/Editor/MeshPaintMode/Private/MeshPaintEdMode.cpp
--- a/Engine/Source/Editor/MeshPaintMode/Private/MeshPaintEdMode.cpp
+++ b/Engine/Source/Editor/MeshPaintMode/Private/MeshPaintEdMode.cpp
@@ -19,4 +19,19 @@ TSharedPtr<class FModeToolkit> FEdModeMeshPaint::GetToolkit()
 	return MakeShareable(new FMeshPaintModeToolKit(this));
 }
 
+bool FEdModeMeshPaint::InputKey(FEditorViewportClient* InViewportClient, FViewport* InViewport, FKey InKey, EInputEvent InEvent)
+{
+	if (!IsEditingEnabled())
+	{
+		return false;
+	}
+
+	return IMeshPaintEdMode::InputKey(InViewportClient, InViewport, InKey, InEvent);
+}
+
+bool FEdModeMeshPaint::IsEditingEnabled() const
+{
+	return GetWorld() ? GetWorld()->FeatureLevel >= ERHIFeatureLevel::SM5 : false;
+}
+
 #undef LOCTEXT_NAMESPACE // "FEdModeMeshPaint"
diff --git a/Engine/Source/Editor/MeshPaintMode/Private/SPaintModeWidget.cpp b/Engine/Source/Editor/MeshPaintMode/Private/SPaintModeWidget.cpp
--- a/Engine/Source/Editor/MeshPaintMode/Private/SPaintModeWidget.cpp
+++ b/Engine/Source/Editor/MeshPaintMode/Private/SPaintModeWidget.cpp
@@ -20,6 +20,10 @@
 #include "PaintModeCommands.h"
 #include "DetailLayoutBuilder.h"
 
+#include "EditorModeManager.h"
+#include "EditorModes.h"
+#include "MeshPaintEdMode.h"
+
 #define LOCTEXT_NAMESPACE "PaintModePainter"
 
 void SPaintModeWidget::Construct(const FArguments& InArgs, FPaintModePainter* InPainter)
@@ -35,10 +39,16 @@ void SPaintModeWidget::Construct(const FArguments& InArgs, FPaintModePainter* In
 	[
 		SNew(SScrollBox)
 		+ SScrollBox::Slot()
+		.Padding(0, 0, 0, 5)
+		[
+			SAssignNew(ErrorTextWidget, SErrorText)
+		]
+		+ SScrollBox::Slot()
 		.Padding(0.0f)
 		[
 			SNew(SVerticalBox)
 			/** Toolbar containing buttons to switch between different paint modes */
+			.IsEnabled(this, &SPaintModeWidget::GetMeshPaintEditorIsEnabled)
 			+ SVerticalBox::Slot()
 			.AutoHeight()
 			[
@@ -256,4 +266,17 @@ void SPaintModeWidget::NotifyPostChange(const FPropertyChangedEvent& PropertyCha
 	}
 }
 
+bool SPaintModeWidget::GetMeshPaintEditorIsEnabled() const
+{
+	FEdModeMeshPaint* MeshPaintMode = (FEdModeMeshPaint*)GLevelEditorModeTools().GetActiveMode(FBuiltinEditorModes::EM_MeshPaint);
+	if (MeshPaintMode)
+	{
+		bool bEnabled = MeshPaintMode->IsEditingEnabled();
+		FText ErrorText = bEnabled ? FText::GetEmpty() : LOCTEXT("MeshPaintSM5Only", "Mesh Paint mode can only be used in SM5.");
+		ErrorTextWidget->SetError(ErrorText);
+		return bEnabled;
+	}
+	return false;
+}
+
 #undef LOCTEXT_NAMESPACE // "PaintModePainter"
diff --git a/Engine/Source/Editor/MeshPaintMode/Private/SPaintModeWidget.h b/Engine/Source/Editor/MeshPaintMode/Private/SPaintModeWidget.h
--- a/Engine/Source/Editor/MeshPaintMode/Private/SPaintModeWidget.h
+++ b/Engine/Source/Editor/MeshPaintMode/Private/SPaintModeWidget.h
@@ -9,6 +9,7 @@
 class FPaintModePainter;
 class UPaintModeSettings;
 class IDetailsView;
+class SErrorText;
 
 /** Widget representing the state / functionality and settings for PaintModePainter*/
 class SPaintModeWidget : public SCompoundWidget, public FNotifyHook
@@ -46,6 +47,10 @@ protected:
 	/** Paint settings instance */
 	UPaintModeSettings* PaintModeSettings;
 
+	TSharedPtr<SErrorText> ErrorTextWidget;
+		
 private:
 	virtual void NotifyPostChange(const FPropertyChangedEvent& PropertyChangedEvent, UProperty* PropertyThatChanged) override;
+	
+	bool GetMeshPaintEditorIsEnabled() const;
 };
diff --git a/Engine/Source/Editor/MeshPaintMode/Public/MeshPaintEdMode.h b/Engine/Source/Editor/MeshPaintMode/Public/MeshPaintEdMode.h
--- a/Engine/Source/Editor/MeshPaintMode/Public/MeshPaintEdMode.h
+++ b/Engine/Source/Editor/MeshPaintMode/Public/MeshPaintEdMode.h
@@ -19,4 +19,9 @@ public:
 	virtual ~FEdModeMeshPaint() {}
 	virtual void Initialize() override;
 	virtual TSharedPtr< FModeToolkit> GetToolkit() override;
+
+	// IMeshPaintEdMode interface.
+	virtual bool InputKey(FEditorViewportClient* InViewportClient, FViewport* InViewport, FKey InKey, EInputEvent InEvent) override;
+
+	bool IsEditingEnabled() const;
 };
\ No newline at end of file
