Fix CIS

diff --git a/Engine/Source/Editor/MeshPaintMode/Private/SPaintModeWidget.h b/Engine/Source/Editor/MeshPaintMode/Private/SPaintModeWidget.h
--- a/Engine/Source/Editor/MeshPaintMode/Private/SPaintModeWidget.h
+++ b/Engine/Source/Editor/MeshPaintMode/Private/SPaintModeWidget.h
@@ -5,6 +5,7 @@
 #include "Widgets/SCompoundWidget.h"
 #include "Widgets/DeclarativeSyntaxSupport.h"
 #include "Misc/NotifyHook.h"
+#include "Widgets/Notifications/SErrorText.h"
 
 class FPaintModePainter;
 class UPaintModeSettings;
