Mesh painter settings are now stored in EditorPerProjectUserSettings.

diff --git a/Engine/Source/Editor/MeshPaintMode/Private/PaintModeSettings.h b/Engine/Source/Editor/MeshPaintMode/Private/PaintModeSettings.h
--- a/Engine/Source/Editor/MeshPaintMode/Private/PaintModeSettings.h
+++ b/Engine/Source/Editor/MeshPaintMode/Private/PaintModeSettings.h
@@ -167,20 +167,20 @@ enum class EPaintMode : uint8
 };
 
 /** Paint mode settings class derives from base mesh painting settings */
-UCLASS()
+UCLASS(Config=EditorPerProjectUserSettings)
 class UPaintModeSettings : public UMeshPaintSettings
 {
 	GENERATED_UCLASS_BODY()
 public:
 
 	static UPaintModeSettings* Get();
-	
+
 	UPROPERTY()
 	EPaintMode PaintMode;
 
-	UPROPERTY(EditAnywhere, Category = VertexPainting, meta=(ShowOnlyInnerProperties))
+	UPROPERTY(EditAnywhere, Config, Category = VertexPainting, meta=(ShowOnlyInnerProperties))
 	FVertexPaintSettings VertexPaintSettings;
 
-	UPROPERTY(EditAnywhere, Category = TexturePainting, meta=(ShowOnlyInnerProperties))
+	UPROPERTY(EditAnywhere, Config, Category = TexturePainting, meta=(ShowOnlyInnerProperties))
 	FTexturePaintSettings TexturePaintSettings;
 };
diff --git a/Engine/Source/Editor/MeshPaintMode/Private/SPaintModeWidget.cpp b/Engine/Source/Editor/MeshPaintMode/Private/SPaintModeWidget.cpp
--- a/Engine/Source/Editor/MeshPaintMode/Private/SPaintModeWidget.cpp
+++ b/Engine/Source/Editor/MeshPaintMode/Private/SPaintModeWidget.cpp
@@ -104,7 +104,7 @@ void SPaintModeWidget::CreateDetailsView()
 		/*bAllowSearch=*/ false,
 		FDetailsViewArgs::HideNameArea,
 		/*bHideSelectionTip=*/ true,
-		/*InNotifyHook=*/ nullptr,
+		/*InNotifyHook=*/ this,
 		/*InSearchInitialKeyFocus=*/ false,
 		/*InViewIdentifier=*/ NAME_None);
 	DetailsViewArgs.DefaultsOnlyVisibility = EEditDefaultsOnlyNodeVisibility::Automatic;
@@ -245,4 +245,15 @@ EVisibility SPaintModeWidget::IsTexturePaintModeVisible() const
 	return (MeshPaintSettings->PaintMode == EPaintMode::Textures) ? EVisibility::Visible : EVisibility::Collapsed;
 }
 
+void SPaintModeWidget::NotifyPostChange(const FPropertyChangedEvent& PropertyChangedEvent, UProperty* PropertyThatChanged)
+{
+	if (PropertyChangedEvent.ChangeType != EPropertyChangeType::Interactive)
+	{
+		for (UObject* Settings : SettingsObjects)
+		{
+			Settings->SaveConfig();
+		}
+	}
+}
+
 #undef LOCTEXT_NAMESPACE // "PaintModePainter"
diff --git a/Engine/Source/Editor/MeshPaintMode/Private/SPaintModeWidget.h b/Engine/Source/Editor/MeshPaintMode/Private/SPaintModeWidget.h
--- a/Engine/Source/Editor/MeshPaintMode/Private/SPaintModeWidget.h
+++ b/Engine/Source/Editor/MeshPaintMode/Private/SPaintModeWidget.h
@@ -4,13 +4,14 @@
 
 #include "Widgets/SCompoundWidget.h"
 #include "Widgets/DeclarativeSyntaxSupport.h"
+#include "Misc/NotifyHook.h"
 
 class FPaintModePainter;
 class UPaintModeSettings;
 class IDetailsView;
 
 /** Widget representing the state / functionality and settings for PaintModePainter*/
-class SPaintModeWidget : public SCompoundWidget
+class SPaintModeWidget : public SCompoundWidget, public FNotifyHook
 {
 public:
 	SLATE_BEGIN_ARGS(SPaintModeWidget) {}
@@ -44,4 +45,7 @@ protected:
 	FPaintModePainter* MeshPainter;
 	/** Paint settings instance */
 	UPaintModeSettings* PaintModeSettings;
+
+private:
+	virtual void NotifyPostChange(const FPropertyChangedEvent& PropertyChangedEvent, UProperty* PropertyThatChanged) override;
 };
