Mesh painting commands are now registered on module load instead of on activation, which allows users to change keybindings before switching to the mode.

diff --git a/Engine/Source/Editor/MeshPaint/Private/IMeshPainter.cpp b/Engine/Source/Editor/MeshPaint/Private/IMeshPainter.cpp
--- a/Engine/Source/Editor/MeshPaint/Private/IMeshPainter.cpp
+++ b/Engine/Source/Editor/MeshPaint/Private/IMeshPainter.cpp
@@ -32,7 +32,6 @@ IMeshPainter::IMeshPainter() :
 	HoverVertexPointColor(0.3f, 1.0f, 0.3f),
 	PaintTransaction(nullptr)
 {
-	FMeshPainterCommands::Register();
 }
 
 IMeshPainter::~IMeshPainter()
diff --git a/Engine/Source/Editor/MeshPaintMode/Private/MeshPaintModeModule.cpp b/Engine/Source/Editor/MeshPaintMode/Private/MeshPaintModeModule.cpp
--- a/Engine/Source/Editor/MeshPaintMode/Private/MeshPaintModeModule.cpp
+++ b/Engine/Source/Editor/MeshPaintMode/Private/MeshPaintModeModule.cpp
@@ -7,6 +7,7 @@
 #include "EditorModeRegistry.h"
 #include "EditorModes.h"
 #include "MeshPaintEdMode.h"
+#include "MeshPainterCommands.h"
 
 #include "PropertyEditorModule.h"
 #include "PaintModeSettingsCustomization.h"
@@ -25,6 +26,8 @@ void FMeshPaintModeModule::StartupModule()
 		FSlateIcon(FEditorStyle::GetStyleSetName(), "LevelEditor.MeshPaintMode", "LevelEditor.MeshPaintMode.Small"),
 		true, 200 );
 
+	FMeshPainterCommands::Register();
+
 	/** Register detail/property customization */
 	FPropertyEditorModule& PropertyModule = FModuleManager::LoadModuleChecked<FPropertyEditorModule>("PropertyEditor");
 	PropertyModule.RegisterCustomClassLayout("PaintModeSettings", FOnGetDetailCustomizationInstance::CreateStatic(&FPaintModeSettingsCustomization::MakeInstance));
@@ -37,6 +40,8 @@ void FMeshPaintModeModule::StartupModule()
 
 void FMeshPaintModeModule::ShutdownModule()
 {
+	FMeshPainterCommands::Unregister();
+
 	FEditorModeRegistry::Get().UnregisterMode(FBuiltinEditorModes::EM_MeshPaint);
 
 	/** De-register detail/property customization */
