Missed header removal

diff --git a/Engine/Source/Editor/MeshPaintMode/Private/MeshPaintModeModule.cpp b/Engine/Source/Editor/MeshPaintMode/Private/MeshPaintModeModule.cpp
--- a/Engine/Source/Editor/MeshPaintMode/Private/MeshPaintModeModule.cpp
+++ b/Engine/Source/Editor/MeshPaintMode/Private/MeshPaintModeModule.cpp
@@ -7,7 +7,6 @@
 #include "EditorModeRegistry.h"
 #include "EditorModes.h"
 #include "MeshPaintEdMode.h"
-#include "MeshPainterCommands.h"
 
 #include "PropertyEditorModule.h"
 #include "PaintModeSettingsCustomization.h"
